Fix integer overflow bug in FIESLoader.

#rnx

diff --git a/Engine/Source/Runtime/IESFile/Private/IESLoader.cpp b/Engine/Source/Runtime/IESFile/Private/IESLoader.cpp
--- a/Engine/Source/Runtime/IESFile/Private/IESLoader.cpp
+++ b/Engine/Source/Runtime/IESFile/Private/IESLoader.cpp
@@ -118,6 +118,9 @@ FIESLoader::FIESLoader(const uint8* Buffer, uint32 BufferLength)
 {
 	check(!IsValid());
 
+	// Make sure BufferLength plus the terminator doesn't overflow the size type used by TArray (which is signed).
+	BufferLength = FMath::Min<uint32>(BufferLength, MAX_int32 - 1);
+
 	TArray<uint8> ASCIIFile;
 
 	// add 0 termination for easier parsing
